Fix AI module load order issue. The AIModule is implicitly loaded early in editor startup from DetailCustomizations, but if you break that dependency link the AI module does not get loaded until PreDefault, which is way after when the Pawn CDO gets initialized.
I was able to fix a NavigationSystemBase case because it only needed a soft ref, but the Pawn needs a hard ref to AIController so AIController must be loaded before Pawn's CDO gets initialized
The AISupportModule now has a hard reference to an AI module symbol, this causes that module to get loaded very early, but the StarutpModule happens at the same time it used to, fairly late in startup

diff --git a/Engine/Plugins/AI/AISupport/AISupport.uplugin b/Engine/Plugins/AI/AISupport/AISupport.uplugin
--- a/Engine/Plugins/AI/AISupport/AISupport.uplugin
+++ b/Engine/Plugins/AI/AISupport/AISupport.uplugin
@@ -19,7 +19,7 @@
         {
             "Name": "AISupportModule",
             "Type": "Runtime",
-            "LoadingPhase": "PreDefault"
+            "LoadingPhase": "PreEarlyLoadingScreen"
         }
 	]
 }
\ No newline at end of file
diff --git a/Engine/Plugins/AI/AISupport/Source/AISupportModule/Private/AISupportModule.cpp b/Engine/Plugins/AI/AISupport/Source/AISupportModule/Private/AISupportModule.cpp
--- a/Engine/Plugins/AI/AISupport/Source/AISupportModule/Private/AISupportModule.cpp
+++ b/Engine/Plugins/AI/AISupport/Source/AISupportModule/Private/AISupportModule.cpp
@@ -2,6 +2,7 @@
 
 #include "AISupportModule.h"
 #include "Modules/ModuleManager.h"
+#include "AITypes.h"
 
 #define LOCTEXT_NAMESPACE "AISupport"
 
@@ -16,17 +17,14 @@ IMPLEMENT_MODULE(FAISupportModule, AISupportModule)
 
 void FAISupportModule::StartupModule()
 {
-	// Called right after the module DLL has been loaded and the module object has been created
-	// Load dependent modules here, and they will be guaranteed to be available during ShutdownModule. ie:
-	// 		FModuleManager::Get().LoadModuleChecked(TEXT("HTTP"));
+	// We need to actually link in something from the AI module, otherwise DLL dependencies will not get loaded early enough
+	// It is NOT safe to actually call startup on the AI module though, as that depends on things that get initialized later in launch
+	static FAIResourcesSet StaticSet;
 }
 
 void FAISupportModule::ShutdownModule()
 {
-	// Called before the module is unloaded, right before the module object is destroyed.
-	// During normal shutdown, this is called in reverse order that modules finish StartupModule().
-	// This means that, as long as a module references dependent modules in it's StartupModule(), it
-	// can safely reference those dependencies in ShutdownModule() as well.
+
 }
 
 #undef LOCTEXT_NAMESPACE
diff --git a/Engine/Source/Runtime/Engine/Private/Pawn.cpp b/Engine/Source/Runtime/Engine/Private/Pawn.cpp
--- a/Engine/Source/Runtime/Engine/Private/Pawn.cpp
+++ b/Engine/Source/Runtime/Engine/Private/Pawn.cpp
@@ -42,6 +42,7 @@ APawn::APawn(const FObjectInitializer& ObjectInitializer)
 
 	if (HasAnyFlags(RF_ClassDefaultObject) && GetClass() == APawn::StaticClass())
 	{
+		// WARNING: This line is why the AISupport plugin has to load the AIModule at PreEarlyLoadingScreen time, otherwise this load fails and CDOs are corrupt in the editor
 		AIControllerClass = LoadClass<AController>(nullptr, *((UEngine*)(UEngine::StaticClass()->GetDefaultObject()))->AIControllerClassName.ToString(), nullptr, LOAD_None, nullptr);
 	}
 	else
diff --git a/Engine/Source/Runtime/NavigationSystem/Private/NavigationSystem.cpp b/Engine/Source/Runtime/NavigationSystem/Private/NavigationSystem.cpp
--- a/Engine/Source/Runtime/NavigationSystem/Private/NavigationSystem.cpp
+++ b/Engine/Source/Runtime/NavigationSystem/Private/NavigationSystem.cpp
@@ -386,10 +386,9 @@ UNavigationSystemV1::UNavigationSystemV1(const FObjectInitializer& ObjectInitial
 	};
 	static FDelegatesInitializer DelegatesInitializer;
 	
-	// @hack, trying to load AIModule's CrowdManager
-	UClass* Class = StaticLoadClass(UCrowdManagerBase::StaticClass(), nullptr, TEXT("/Script/AIModule.CrowdManager"));
-	CrowdManagerClass = Class ? Class : UCrowdManagerBase::StaticClass();
-	
+	// Set to the ai module's crowd manager, this module may not exist at spawn time but then it will just fail to load
+	CrowdManagerClass = FSoftObjectPath(TEXT("/Script/AIModule.CrowdManager"));
+
 	// active tiles
 	NextInvokersUpdateTime = 0.f;
 	ActiveTilesUpdateInterval = 1.f;
