[NO NON-PII COMMIT MESSAGE]

diff --git a/Engine/Plugins/AI/AISupport/Source/AISupportModule/Private/AISupportModule.cpp b/Engine/Plugins/AI/AISupport/Source/AISupportModule/Private/AISupportModule.cpp
--- a/Engine/Plugins/AI/AISupport/Source/AISupportModule/Private/AISupportModule.cpp
+++ b/Engine/Plugins/AI/AISupport/Source/AISupportModule/Private/AISupportModule.cpp
@@ -19,7 +19,8 @@ void FAISupportModule::StartupModule()
 {
 	// We need to actually link in something from the AI module, otherwise DLL dependencies will not get loaded early enough
 	// It is NOT safe to actually call startup on the AI module though, as that depends on things that get initialized later in launch
-	static FAIResourcesSet StaticSet;
+	static int32 ForceLink = FAIResources::GetResourcesCount();
+	ForceLink++;
 }
 
 void FAISupportModule::ShutdownModule()
