[NO NON-PII COMMIT MESSAGE]

diff --git a/Engine/Plugins/AI/AISupport/AISupport.uplugin b/Engine/Plugins/AI/AISupport/AISupport.uplugin
--- a/Engine/Plugins/AI/AISupport/AISupport.uplugin
+++ b/Engine/Plugins/AI/AISupport/AISupport.uplugin
@@ -19,7 +19,7 @@
         {
             "Name": "AISupportModule",
             "Type": "Runtime",
-            "LoadingPhase": "PreEarlyLoadingScreen"
+            "LoadingPhase": "PostConfigInit"
         }
 	]
 }
\ No newline at end of file
